Enter the beginning phase of the Interactable Interface plugin. Still in active development!

#rnx

diff --git a/Engine/Plugins/Experimental/InteractionInterface/InteractionInterface.uplugin b/Engine/Plugins/Experimental/InteractionInterface/InteractionInterface.uplugin
new file mode 100644
--- /dev/null
+++ b/Engine/Plugins/Experimental/InteractionInterface/InteractionInterface.uplugin
@@ -0,0 +1,34 @@
+{
+	"FileVersion": 3,
+	"Version": 1,
+	"VersionName": "1.0",
+	"Description": "Interaction Interface",
+	"Category": "Gameplay",
+	"CreatedBy": "Epic Games, Inc.",
+	"CreatedByURL": "https://epicgames.com",
+	"DocsURL": "",
+	"MarketplaceURL": "",
+	"SupportURL": "",
+	"EnabledByDefault": false,
+	"CanContainContent": true,
+	"IsBetaVersion": false,
+	"IsExperimentalVersion": true,
+	"Installed": false,
+	"Modules": [
+		{
+			"Name": "InteractableInterface",
+			"Type": "Runtime",
+			"LoadingPhase": "PreDefault"
+		}
+	],
+	"Plugins": [
+		{
+			"Name": "SmartObjects",
+			"Enabled": true
+		},
+		{
+			"Name": "GameplayAbilities",
+			"Enabled": true
+		}
+	]
+}
\ No newline at end of file
diff --git a/Engine/Plugins/Experimental/InteractionInterface/Source/InteractableInterface/InteractableInterface.build.cs b/Engine/Plugins/Experimental/InteractionInterface/Source/InteractableInterface/InteractableInterface.build.cs
new file mode 100644
--- /dev/null
+++ b/Engine/Plugins/Experimental/InteractionInterface/Source/InteractableInterface/InteractableInterface.build.cs
@@ -0,0 +1,21 @@
+// Copyright Epic Games, Inc. All Rights Reserved.
+
+namespace UnrealBuildTool.Rules
+{
+	public class InteractableInterface : ModuleRules
+	{
+		public InteractableInterface(ReadOnlyTargetRules Target) : base(Target)
+		{
+			PublicDependencyModuleNames.AddRange(
+				new string[] {
+					"Core",
+					"CoreUObject",
+					"Engine",
+					"GameplayTags",
+					"GameplayAbilities",
+					"GameplayTasks"
+				}
+			);
+		}
+	}
+}
\ No newline at end of file
diff --git a/Engine/Plugins/Experimental/InteractionInterface/Source/InteractableInterface/Private/InteractableInterfaceLibrary.cpp b/Engine/Plugins/Experimental/InteractionInterface/Source/InteractableInterface/Private/InteractableInterfaceLibrary.cpp
new file mode 100644
--- /dev/null
+++ b/Engine/Plugins/Experimental/InteractionInterface/Source/InteractableInterface/Private/InteractableInterfaceLibrary.cpp
@@ -0,0 +1,117 @@
+// Copyright Epic Games, Inc. All Rights Reserved.
+
+#include "InteractableInterfaceLibrary.h"
+
+// Interaction interface module
+#include "InteractableTargetInterface.h"
+#include "InteractionInterfaceLogs.h"
+
+// Engine
+#include "Components/ActorComponent.h"
+#include "Components/PrimitiveComponent.h"
+#include "Engine/OverlapResult.h"
+#include "GameFramework/Actor.h"
+
+void UInteractableInterfaceLibrary::GetInteractableTargetsFromActor(AActor* Actor, TArray<TScriptInterface<IInteractionTarget>>& OutInteractionTargets)
+{
+	TRACE_CPUPROFILER_EVENT_SCOPE(UInteractableInterfaceLibrary::GetInteractableTargetsFromActor);
+	
+	// If the actor is directly interactable, return that.
+	const TScriptInterface<IInteractionTarget> InteractableActor(Actor);
+	if (InteractableActor)
+	{
+		OutInteractionTargets.Add(InteractableActor);
+	}
+
+	// If the actor isn't interactable, it might have a component that has a interactable interface.
+	TArray<UActorComponent*> InteractableComponents = Actor ? Actor->GetComponentsByInterface(UInteractionTarget::StaticClass()) : TArray<UActorComponent*>();
+	for (UActorComponent* InteractableComponent : InteractableComponents)
+	{
+		OutInteractionTargets.Add(TScriptInterface<IInteractionTarget>(InteractableComponent));
+	}
+}
+
+void UInteractableInterfaceLibrary::AppendTargetConfiguration(TScriptInterface<IInteractionTarget> Target, const FInteractionContext& Context, FInteractionQueryResults& OutResults)
+{
+	if (!Target)
+	{
+		UE_LOG(LogInteractions, Error, TEXT("[%hs] Invalid target! Exiting."), __func__);
+		return;
+	}
+	
+	Target->AppendTargetConfiguration(Context, OutResults);
+}
+
+void UInteractableInterfaceLibrary::BeginInteractionOnTarget(TScriptInterface<IInteractionTarget> Target, const FInteractionContext& Context)
+{
+	if (!Target)
+	{
+		UE_LOG(LogInteractions, Error, TEXT("[%hs] Invalid target! Exiting."), __func__);
+		return;
+	}
+
+	Target->BeginInteraction(Context);
+}
+
+void UInteractableInterfaceLibrary::ResetQueryResults(UPARAM(ref) FInteractionQueryResults& ToReset)
+{
+	ToReset.Reset();
+}
+
+void UInteractableInterfaceLibrary::AppendInteractableTargetsFromOverlapResults(const TArray<FOverlapResult>& OverlapResults, TArray<TScriptInterface<IInteractionTarget>>& OutInteractableTargets)
+{
+	TRACE_CPUPROFILER_EVENT_SCOPE(UInteractableInterfaceLibrary::AppendInteractableTargetsFromOverlapResults);
+	
+	for (const FOverlapResult& Overlap : OverlapResults)
+	{
+		const TScriptInterface<IInteractionTarget> InteractableActor(Overlap.GetActor());
+		if (InteractableActor)
+		{
+		 	OutInteractableTargets.AddUnique(InteractableActor);
+		}
+		
+		const TScriptInterface<IInteractionTarget> InteractableComponent(Overlap.GetComponent());
+		if (InteractableComponent)
+		{
+			OutInteractableTargets.AddUnique(InteractableComponent);
+		}
+	}
+}
+
+void UInteractableInterfaceLibrary::AppendInteractableTargetsFromHitResult(const FHitResult& HitResult, TArray<TScriptInterface<IInteractionTarget>>& OutInteractableTargets)
+{
+	TRACE_CPUPROFILER_EVENT_SCOPE(UInteractableInterfaceLibrary::AppendInteractableTargetsFromHitResult);
+	
+	TScriptInterface<IInteractionTarget> InteractableActor(HitResult.GetActor());
+	if (InteractableActor)
+	{
+		OutInteractableTargets.AddUnique(InteractableActor);
+	}
+	
+	TScriptInterface<IInteractionTarget> InteractableComponent(HitResult.GetComponent());
+	if (InteractableComponent)
+	{
+		OutInteractableTargets.AddUnique(InteractableComponent);
+	}
+}
+
+AActor* UInteractableInterfaceLibrary::GetActorFromInteractableTarget(const TScriptInterface<IInteractionTarget> InteractableTarget)
+{
+	if (UObject* Object = InteractableTarget.GetObject())
+	{
+		if (AActor* Actor = Cast<AActor>(Object))
+		{
+			return Actor;
+		}
+		else if (UActorComponent* ActorComponent = Cast<UActorComponent>(Object))
+		{
+			return ActorComponent->GetOwner();
+		}
+		else
+		{
+			unimplemented();
+		}
+	}
+
+	return nullptr;
+}
diff --git a/Engine/Plugins/Experimental/InteractionInterface/Source/InteractableInterface/Private/InteractionInstigatorComponent.cpp b/Engine/Plugins/Experimental/InteractionInterface/Source/InteractableInterface/Private/InteractionInstigatorComponent.cpp
new file mode 100644
--- /dev/null
+++ b/Engine/Plugins/Experimental/InteractionInterface/Source/InteractableInterface/Private/InteractionInstigatorComponent.cpp
@@ -0,0 +1,47 @@
+// Copyright Epic Games, Inc. All Rights Reserved.
+
+#include "InteractionInstigatorComponent.h"
+
+#include "InteractableTargetInterface.h"
+#include "InteractionTypes.h"
+#include "InteractionInterfaceLogs.h"
+
+UInteractionInstigatorComponent::UInteractionInstigatorComponent(const FObjectInitializer& ObjectInitializer)
+	: Super(ObjectInitializer)
+{
+	
+}
+
+void UInteractionInstigatorComponent::AttemptToBeginInteractions(const TArray<TScriptInterface<IInteractionTarget>>& TargetsToInteractWith)
+{
+	OnAttemptToBeginInteractions(TargetsToInteractWith);
+}
+
+void UInteractionInstigatorComponent::OnAttemptToBeginInteractions(const TArray<TScriptInterface<IInteractionTarget>>& TargetsToInteractWith)
+{
+	TRACE_CPUPROFILER_EVENT_SCOPE(UInteractionInstigatorComponent::OnAttemptToBeginInteractions);
+	
+	UE_LOG(LogInteractions, Log, TEXT("[%hs] %s is attempting to being interaction with targets..."), __func__, *GetNameSafe(this));
+
+	FInteractionContextData ContextData = {};
+	ContextData.Instigator = this;
+	
+	// You can make any kind of interaction target data here that you prefer
+	FInteractionContext Context = {};
+
+	// The context data about this interaction. Provide some data about your instigator here.
+	// Maybe there is a specific smart object slot handle you would like to interact with, or
+	// other conditions you would like your target to know about. 
+	Context.ContextData = InteractionContextData;
+
+	// Here you could call "IInteractionTarget::AppendTargetConfiguration" if you wanted to on each target.
+	// This would allow you to begin interactions conditionally based on some criteria that you set up for your 
+	// game if desired.
+
+	// Get info about how we are supposed to respond to any interactions and metadata about them
+	for (const TScriptInterface<IInteractionTarget>& Target : TargetsToInteractWith)
+	{
+		Context.Target = Target;
+		Target->BeginInteraction(Context);
+	}
+}
diff --git a/Engine/Plugins/Experimental/InteractionInterface/Source/InteractableInterface/Private/InteractionInterfaceLogs.cpp b/Engine/Plugins/Experimental/InteractionInterface/Source/InteractableInterface/Private/InteractionInterfaceLogs.cpp
new file mode 100644
--- /dev/null
+++ b/Engine/Plugins/Experimental/InteractionInterface/Source/InteractableInterface/Private/InteractionInterfaceLogs.cpp
@@ -0,0 +1,5 @@
+// Copyright Epic Games, Inc. All Rights Reserved.
+
+#include "InteractionInterfaceLogs.h"
+
+DEFINE_LOG_CATEGORY(LogInteractions);
\ No newline at end of file
diff --git a/Engine/Plugins/Experimental/InteractionInterface/Source/InteractableInterface/Private/InteractionInterfaceModule.cpp b/Engine/Plugins/Experimental/InteractionInterface/Source/InteractableInterface/Private/InteractionInterfaceModule.cpp
new file mode 100644
--- /dev/null
+++ b/Engine/Plugins/Experimental/InteractionInterface/Source/InteractableInterface/Private/InteractionInterfaceModule.cpp
@@ -0,0 +1,18 @@
+// Copyright Epic Games, Inc. All Rights Reserved.
+
+#include "InteractionInterfaceModule.h"
+
+#include "Modules/ModuleManager.h"
+
+FInteractionInterfaceModule& FInteractionInterfaceModule::Get()
+{
+	static FInteractionInterfaceModule& Singleton = FModuleManager::LoadModuleChecked<FInteractionInterfaceModule>("InteractionInterface");
+	return Singleton;
+}
+
+bool FInteractionInterfaceModule::IsAvailable()
+{
+	return FModuleManager::Get().IsModuleLoaded("InteractionInterface");
+}
+
+IMPLEMENT_MODULE(FInteractionInterfaceModule, InteractableInterface)
\ No newline at end of file
diff --git a/Engine/Plugins/Experimental/InteractionInterface/Source/InteractableInterface/Private/InteractionInterfaceModule.h b/Engine/Plugins/Experimental/InteractionInterface/Source/InteractableInterface/Private/InteractionInterfaceModule.h
new file mode 100644
--- /dev/null
+++ b/Engine/Plugins/Experimental/InteractionInterface/Source/InteractableInterface/Private/InteractionInterfaceModule.h
@@ -0,0 +1,25 @@
+// Copyright Epic Games, Inc. All Rights Reserved.
+
+#pragma once
+
+#include "Modules/ModuleInterface.h"
+
+class FInteractionInterfaceModule : public IModuleInterface
+{
+public:
+
+	/**
+	 * Singleton-like access to this module's interface.  This is just for convenience!
+	 * Beware of calling this during the shutdown phase, though.  Your module might have been unloaded already.
+	 *
+	 * @return Returns singleton instance, loading the module on demand if needed
+	 */
+	static FInteractionInterfaceModule& Get();
+
+	/**
+	 * Checks to see if this module is loaded and ready.  It is only valid to call Get() if IsAvailable() returns true.
+	 *
+	 * @return True if the module is loaded and ready to use
+	 */
+	static bool IsAvailable();
+};
\ No newline at end of file
diff --git a/Engine/Plugins/Experimental/InteractionInterface/Source/InteractableInterface/Private/InteractionTargetComponent.cpp b/Engine/Plugins/Experimental/InteractionInterface/Source/InteractableInterface/Private/InteractionTargetComponent.cpp
new file mode 100644
--- /dev/null
+++ b/Engine/Plugins/Experimental/InteractionInterface/Source/InteractableInterface/Private/InteractionTargetComponent.cpp
@@ -0,0 +1,41 @@
+// Copyright Epic Games, Inc. All Rights Reserved.
+
+#include "InteractionTargetComponent.h"
+#include "InteractionTypes.h"
+#include "InteractionInterfaceLogs.h"
+
+UInteractionTargetComponent::UInteractionTargetComponent(const FObjectInitializer& ObjectInitializer)
+	: Super(ObjectInitializer)
+{
+
+}
+
+void UInteractionTargetComponent::BP_AppendTargetConfiguration(const FInteractionContext& Context, FInteractionQueryResults& OutResults) const
+{
+	// Just call in the interface implementation on this blueprint.
+	AppendTargetConfiguration(Context, OutResults);
+}
+
+void UInteractionTargetComponent::BP_BeginInteraction(const FInteractionContext& Context)
+{
+	// Just call in the interface implementation on this blueprint.
+	BeginInteraction(Context);
+}
+
+void UInteractionTargetComponent::AppendTargetConfiguration(const FInteractionContext& QueryContext, FInteractionQueryResults& OutResults) const
+{
+	TRACE_CPUPROFILER_EVENT_SCOPE(UInteractionTargetComponent::AppendTargetConfiguration);
+	
+	// Add any info about the interactions that are available on this target.
+	OutResults.AvailableInteractions.Append(TargetConfigs);
+}
+
+void UInteractionTargetComponent::BeginInteraction(const FInteractionContext& Context)
+{
+	TRACE_CPUPROFILER_EVENT_SCOPE(UInteractionTargetComponent::BeginInteraction);
+	
+	UE_LOG(LogInteractions, Log, TEXT("[%hs] Native C++ impl of begin interaction on this target %s"), __func__, *GetNameSafe(this));
+
+	// Broadcast BP event that you can bind to
+	OnBeginInteractionCallback.Broadcast(Context);
+}
diff --git a/Engine/Plugins/Experimental/InteractionInterface/Source/InteractableInterface/Private/InteractionTask_WaitForTargets.cpp b/Engine/Plugins/Experimental/InteractionInterface/Source/InteractableInterface/Private/InteractionTask_WaitForTargets.cpp
new file mode 100644
--- /dev/null
+++ b/Engine/Plugins/Experimental/InteractionInterface/Source/InteractableInterface/Private/InteractionTask_WaitForTargets.cpp
@@ -0,0 +1,179 @@
+// Copyright Epic Games, Inc. All Rights Reserved.
+
+#include "InteractionTask_WaitForTargets.h"
+
+// Engine
+#include "AbilitySystemComponent.h"
+#include "Engine/OverlapResult.h"
+#include "Engine/World.h"
+#include "GameFramework/Controller.h"
+#include "TimerManager.h"
+#include "DrawDebugHelpers.h"
+
+
+// Interactable Interface
+#include "InteractableInterfaceLibrary.h"
+#include "InteractionInterfaceLogs.h"
+
+//////////////////////////////////////////////////////////
+// UAbilityTask_GrantNearbyInteractionData
+
+UAbilityTask_GrantNearbyInteractionData* UAbilityTask_GrantNearbyInteractionData::GrantAbilitiesForNearbyInteractionData(
+	UGameplayAbility* OwningAbility,
+	ECollisionChannel TraceChannel,
+	float InteractionScanRange /* = 500.f */,
+	float InteractionScanRate /* = 0.1f */)
+{
+	UAbilityTask_GrantNearbyInteractionData* MyObj = NewAbilityTask<UAbilityTask_GrantNearbyInteractionData>(OwningAbility);
+	MyObj->InteractionScanRange = InteractionScanRange;
+	MyObj->InteractionScanRate = InteractionScanRate;
+	MyObj->InteractionTraceChannel = TraceChannel;
+	
+	return MyObj;
+}
+
+void UAbilityTask_GrantNearbyInteractionData::Activate()
+{
+	// Start scanning for nearby targets at our scan rate
+	SetWaitingOnAvatar();
+	
+	UWorld* World = GetWorld();
+	if (!ensure(World))
+	{
+		return;
+	}
+
+	// Start a time for the scan rate of when the gather interactable target data
+	World->GetTimerManager().SetTimer(
+		QueryTimerHandle,
+		this,
+		&UAbilityTask_GrantNearbyInteractionData::QueryNearbyInteractables,
+		InteractionScanRate,
+		/* bLoop= */ true);
+}
+
+void UAbilityTask_GrantNearbyInteractionData::OnDestroy(bool AbilityEnded)
+{
+	// Clear the scan timer
+	if (UWorld* World = GetWorld())
+	{
+		World->GetTimerManager().ClearTimer(QueryTimerHandle);
+	}
+
+	Super::OnDestroy(AbilityEnded);
+}
+
+void UAbilityTask_GrantNearbyInteractionData::QueryNearbyInteractables()
+{
+	UWorld* World = GetWorld();
+	AActor* ActorOwner = GetAvatarActor();
+
+	if (!World || !ActorOwner)
+	{
+		return;
+	}
+
+	// TODO: Make this an Async overlap check
+	// FTraceHandle handle = World->AsyncOverlapByChannel
+	// if ! handle valid... do another trace.
+	// Leaving this as a TODO to reduce complexity and make sure this stuff really works well before optimizing
+	
+	// Do a sphere trace around the requesting actor and gather the available interactions around them
+	FCollisionQueryParams Params(SCENE_QUERY_STAT(UAbilityTask_GrantNearbyInteractionData), false);
+
+	// Ignore the owner of this ability...
+	Params.AddIgnoredActor(ActorOwner);
+
+	TArray<FOverlapResult> OverlapResults;
+
+	World->OverlapMultiByChannel(
+		OUT OverlapResults,
+		ActorOwner->GetActorLocation(),
+		FQuat::Identity,
+		InteractionTraceChannel,
+		FCollisionShape::MakeSphere(InteractionScanRange),
+		Params);
+		
+#if ENABLE_DRAW_DEBUG
+	DrawDebugSphere(World, ActorOwner->GetActorLocation(), InteractionScanRange, 10, FColor::Green, false, InteractionScanRate);
+#endif
+	
+	if (OverlapResults.IsEmpty())
+	{
+		return;
+	}
+
+	// Gather the IInteractionTarget interface from our query result
+	TArray<TScriptInterface<IInteractionTarget>> InteractableTargets;
+	UInteractableInterfaceLibrary::AppendInteractableTargetsFromOverlapResults(OverlapResults, OUT InteractableTargets);
+
+	// TODO: We could make a smart object query here to gather available smart objects in the given area
+	// or make a subclass of this task which does a query using smart object world conditions
+
+	// If the available targets from the spatial query has changed, then broadcast our delegate that
+	// there has been a change in what targets are available
+	if (InteractableTargets != CurrentAvailableTargets)
+	{
+		CurrentAvailableTargets = InteractableTargets;
+		OnAvailableInteractionTargetsChanged.Broadcast(InteractableTargets);
+	}
+}
+
+//////////////////////////////////////////////////////////
+// UGameplayAbility_Interact 
+UGameplayAbility_Interact::UGameplayAbility_Interact(const FObjectInitializer& ObjectInitializer)
+{
+	InstancingPolicy = EGameplayAbilityInstancingPolicy::InstancedPerActor;
+	NetExecutionPolicy = EGameplayAbilityNetExecutionPolicy::LocalPredicted;
+}
+
+void UGameplayAbility_Interact::ActivateAbility(const FGameplayAbilitySpecHandle Handle, const FGameplayAbilityActorInfo* ActorInfo, const FGameplayAbilityActivationInfo ActivationInfo, const FGameplayEventData* TriggerEventData)
+{
+	Super::ActivateAbility(Handle, ActorInfo, ActivationInfo, TriggerEventData);
+
+	const UAbilitySystemComponent* AbilitySystem = GetAbilitySystemComponentFromActorInfo();
+
+	// Create a task which will check for nearby interaction targets! 
+	if (AbilitySystem && AbilitySystem->GetOwnerRole() == ROLE_Authority)
+	{
+		UAbilityTask_GrantNearbyInteractionData* Task = UAbilityTask_GrantNearbyInteractionData::GrantAbilitiesForNearbyInteractionData(
+			this,
+			InteractionTraceChannel,
+			InteractionScanRange,
+			InteractionScanRate);
+
+		Task->OnAvailableInteractionTargetsChanged.AddUniqueDynamic(this, &UGameplayAbility_Interact::HandleTargetsUpdatedFromTask);
+		
+		Task->ReadyForActivation();
+	}
+}
+
+void UGameplayAbility_Interact::HandleTargetsUpdatedFromTask(const TArray<TScriptInterface<IInteractionTarget>>& AvailableTargets)
+{
+	UpdateInteractions(AvailableTargets);	
+}
+
+void UGameplayAbility_Interact::OnAvailableInteractionsUpdated_Implementation()
+{
+	UE_LOG(LogInteractions, Log, TEXT("[%hs] %d targets are now available to be interacted with."), __func__, CurrentAvailableTargets.Num());
+}
+
+void UGameplayAbility_Interact::OnTriggerInteraction_Implementation()
+{
+	UE_LOG(LogInteractions, Log, TEXT("[%hs] trigger interactions on %d targets"), __func__, CurrentAvailableTargets.Num());
+}
+
+void UGameplayAbility_Interact::UpdateInteractions(const TArray<TScriptInterface<IInteractionTarget>>& AvailableTargets)
+{
+	CurrentAvailableTargets = AvailableTargets;
+
+	// Notify BP/C++ that the available interactions have changed
+	OnAvailableInteractionsUpdated();
+}
+
+void UGameplayAbility_Interact::TriggerInteraction()
+{
+	// By default for now just let the BP implementations or C++ subclasses of this ability
+	// decide how or which targets to interact with.
+	OnTriggerInteraction();
+}
\ No newline at end of file
diff --git a/Engine/Plugins/Experimental/InteractionInterface/Source/InteractableInterface/Public/InteractableInstigator.h b/Engine/Plugins/Experimental/InteractionInterface/Source/InteractableInterface/Public/InteractableInstigator.h
new file mode 100644
--- /dev/null
+++ b/Engine/Plugins/Experimental/InteractionInterface/Source/InteractableInterface/Public/InteractableInstigator.h
@@ -0,0 +1,30 @@
+// Copyright Epic Games, Inc. All Rights Reserved.
+
+#pragma once
+
+#include "UObject/Interface.h"
+#include "InteractableInstigator.generated.h"
+
+class IInteractionTarget;
+
+// This class does not need to be modified.
+UINTERFACE(meta=(CannotImplementInterfaceInBlueprint))
+class INTERACTABLEINTERFACE_API UInteractableInstigator : public UInterface
+{
+	GENERATED_BODY()
+};
+
+/**
+ * Interface that represents a single "Instigator" of an interaction.
+ * 
+ * Instigators are objects which are allowed to interact with targets. They
+ * can attempt to begin interacting with a given set of targets.
+ */
+class INTERACTABLEINTERFACE_API IInteractableInstigator
+{
+	GENERATED_BODY()
+	
+protected:
+	
+	virtual void OnAttemptToBeginInteractions(const TArray<TScriptInterface<IInteractionTarget>>& TargetsToInteractWith) = 0;
+};
\ No newline at end of file
diff --git a/Engine/Plugins/Experimental/InteractionInterface/Source/InteractableInterface/Public/InteractableInterfaceLibrary.h b/Engine/Plugins/Experimental/InteractionInterface/Source/InteractableInterface/Public/InteractableInterfaceLibrary.h
new file mode 100644
--- /dev/null
+++ b/Engine/Plugins/Experimental/InteractionInterface/Source/InteractableInterface/Public/InteractableInterfaceLibrary.h
@@ -0,0 +1,80 @@
+// Copyright Epic Games, Inc. All Rights Reserved.
+
+#pragma once
+
+#include "Kismet/BlueprintFunctionLibrary.h"
+#include "InteractionTypes.h"
+
+#include "InteractableInterfaceLibrary.generated.h"
+
+class AActor;
+class IInteractionTarget;
+struct FHitResult;
+struct FOverlapResult;
+
+/**
+ * BP function helpers to utilize the Interactable Interface.
+ */
+UCLASS()
+class INTERACTABLEINTERFACE_API UInteractableInterfaceLibrary final : public UBlueprintFunctionLibrary
+{
+	GENERATED_BODY()
+public:
+
+	/**
+	 * Populates the given OutInteractionTargets array with all IInteractionTarget interface
+	 * objects on the given actor.
+	 *
+	 * @param Actor						The actor to check for targets on
+	 * @param OutInteractionTargets		Array of pointers to any interaction targets on the given actor
+	 */
+	UFUNCTION(BlueprintCallable, Category="Interactions")
+	static void GetInteractableTargetsFromActor(AActor* Actor, TArray<TScriptInterface<IInteractionTarget>>& OutInteractionTargets);
+
+	/**
+	 * Determines what the configuration of this target is. 
+	 * Gather information about this specific target so that it can be displayed
+	 * to the player and provide access to what behavior should occur in response
+	 * to this interaction.
+	 *
+	 * @param Target		The interaction target to get the configuration of
+	 * @param Context		The context of this interaction.
+	 * @param OutResults	Output results that this target will populate
+	 */
+	UFUNCTION(BlueprintCallable, Category="Interactions", meta=(AutoCreateRefTerm="Context"))
+	static void AppendTargetConfiguration(TScriptInterface<IInteractionTarget> Target, const FInteractionContext& Context, UPARAM(ref) FInteractionQueryResults& OutResults);
+
+	/**
+	 * Called when this target is interacted with. Implement any state changes or gameplay affects
+	 * you want this interaction to have here
+	 *
+	 * @param Target		The interaction target to begin interacting with
+	 * @param Context		The context of this interaction. 
+	 */
+	UFUNCTION(BlueprintCallable, Category="Interactions", meta=(AutoCreateRefTerm="Context"))
+	static void BeginInteractionOnTarget(TScriptInterface<IInteractionTarget> Target, const FInteractionContext& Context);
+	
+	/**
+	* Resets the values of the given interaction query results to be empty.
+	*/
+	UFUNCTION(BlueprintCallable, Category="Interactions")
+	static void ResetQueryResults(UPARAM(ref) FInteractionQueryResults& ToReset);
+
+	/**
+	 * Given a set of overlap results, append any found IInteractionTargets to the OutInteractableTargets array.
+	 *
+	 * @param OverlapResults			The overlap results from a physics query to process.
+	 * @param OutInteractableTargets	The array of targets to append any found targets to. This array will not be reset.
+	 */
+	static void AppendInteractableTargetsFromOverlapResults(const TArray<FOverlapResult>& OverlapResults, TArray<TScriptInterface<IInteractionTarget>>& OutInteractableTargets);
+
+	/**
+	 * Given a hit result, append any found IInteractionTargets to the OutInteractableTargets array.
+	 *
+	 * @param HitResult					The hit result to process for targets
+	 * @param OutInteractableTargets	The array of targets to append any found targets to. This array will not be reset.
+	 */
+	static void AppendInteractableTargetsFromHitResult(const FHitResult& HitResult, TArray<TScriptInterface<IInteractionTarget>>& OutInteractableTargets);
+
+	static AActor* GetActorFromInteractableTarget(const TScriptInterface<IInteractionTarget> InteractableTarget);
+};
\ No newline at end of file
diff --git a/Engine/Plugins/Experimental/InteractionInterface/Source/InteractableInterface/Public/InteractableTargetInterface.h b/Engine/Plugins/Experimental/InteractionInterface/Source/InteractableInterface/Public/InteractableTargetInterface.h
new file mode 100644
--- /dev/null
+++ b/Engine/Plugins/Experimental/InteractionInterface/Source/InteractableInterface/Public/InteractableTargetInterface.h
@@ -0,0 +1,49 @@
+// Copyright Epic Games, Inc. All Rights Reserved.
+
+#pragma once
+
+#include "UObject/Interface.h"
+#include "InteractableTargetInterface.generated.h"
+
+struct FInteractionContext;
+struct FInteractionQueryResults;
+
+// This class does not need to be modified.
+UINTERFACE(BlueprintType, meta = (CannotImplementInterfaceInBlueprint))
+class INTERACTABLEINTERFACE_API UInteractionTarget : public UInterface
+{
+	GENERATED_BODY()
+};
+
+/**
+ * Implement this interface on things that are able to be interacted with.
+ *
+ * A target should be something in the world that can be interacted with. For example,
+ * a treasure chest which can be opened, an item that can picked up, or some narrative design
+ * element for your game.
+ */
+class INTERACTABLEINTERFACE_API IInteractionTarget
+{
+	GENERATED_BODY()
+public:
+	
+	/**
+	 * Determines what the configuration of this target is. 
+	 * Gather information about this specific target so that it can be displayed
+	 * to the player and provide access to what behavior should occur in response
+	 * to this interaction.
+	 *
+	 * @param Context		The context of this interaction.
+	 * @param OutResults	Output results that this target will populate
+	 */
+	virtual void AppendTargetConfiguration(const FInteractionContext& Context, OUT FInteractionQueryResults& OutResults) const = 0;
+
+	/**
+	 * Called when this target is interacted with. Implement any state changes or gameplay affects
+	 * you want this interaction to have here
+	 *
+	 * @param Context		The context of this interaction. This is customizable for your
+	 *						game by adding additional context types.
+	 */
+	virtual void BeginInteraction(const FInteractionContext& Context) = 0;	
+};
\ No newline at end of file
diff --git a/Engine/Plugins/Experimental/InteractionInterface/Source/InteractableInterface/Public/InteractionInstigatorComponent.h b/Engine/Plugins/Experimental/InteractionInterface/Source/InteractableInterface/Public/InteractionInstigatorComponent.h
new file mode 100644
--- /dev/null
+++ b/Engine/Plugins/Experimental/InteractionInterface/Source/InteractableInterface/Public/InteractionInstigatorComponent.h
@@ -0,0 +1,41 @@
+// Copyright Epic Games, Inc. All Rights Reserved.
+
+#pragma once
+
+#include "InteractableInstigator.h"
+#include "Components/ActorComponent.h"
+#include "StructUtils/InstancedStruct.h"
+
+#include "InteractionInstigatorComponent.generated.h"
+
+class IInteractionTarget;
+
+UCLASS(Blueprintable, ClassGroup = Gameplay, meta = (BlueprintSpawnableComponent), HideCategories = (Activation, AssetUserData, Collision))
+class INTERACTABLEINTERFACE_API UInteractionInstigatorComponent :
+	public UActorComponent,
+	public IInteractableInstigator
+{
+	GENERATED_BODY()
+	
+public:
+
+	UInteractionInstigatorComponent(const FObjectInitializer& ObjectInitializer);
+
+	/**
+	 * Attempts to begin interacting with the given array of targets
+	 */
+	UFUNCTION(BlueprintCallable, Category="Interactions")
+	void AttemptToBeginInteractions(const TArray<TScriptInterface<IInteractionTarget>>& TargetsToInteractWith);
+	
+protected:
+
+	/**
+	* Data about this specific interaction
+	*/
+	UPROPERTY(BlueprintReadWrite, EditAnywhere, Category="Interactions", meta = (BaseStruct = "/Script/InteractableInterface.InteractionContextData"))
+	FInstancedStruct InteractionContextData;
+
+	//~Begin IInteractableInstigator interface
+	virtual void OnAttemptToBeginInteractions(const TArray<TScriptInterface<IInteractionTarget>>& TargetsToInteractWith) override;
+	//~End IInteractableInstigator interface
+};
\ No newline at end of file
diff --git a/Engine/Plugins/Experimental/InteractionInterface/Source/InteractableInterface/Public/InteractionInterfaceLogs.h b/Engine/Plugins/Experimental/InteractionInterface/Source/InteractableInterface/Public/InteractionInterfaceLogs.h
new file mode 100644
--- /dev/null
+++ b/Engine/Plugins/Experimental/InteractionInterface/Source/InteractableInterface/Public/InteractionInterfaceLogs.h
@@ -0,0 +1,7 @@
+// Copyright Epic Games, Inc. All Rights Reserved.
+
+#pragma once
+
+#include "Logging/LogMacros.h"
+
+INTERACTABLEINTERFACE_API DECLARE_LOG_CATEGORY_EXTERN(LogInteractions, Log, All);
\ No newline at end of file
diff --git a/Engine/Plugins/Experimental/InteractionInterface/Source/InteractableInterface/Public/InteractionTargetComponent.h b/Engine/Plugins/Experimental/InteractionInterface/Source/InteractableInterface/Public/InteractionTargetComponent.h
new file mode 100644
--- /dev/null
+++ b/Engine/Plugins/Experimental/InteractionInterface/Source/InteractableInterface/Public/InteractionTargetComponent.h
@@ -0,0 +1,71 @@
+// Copyright Epic Games, Inc. All Rights Reserved.
+
+#pragma once
+
+#include "InteractableTargetInterface.h"
+#include "Components/ActorComponent.h"
+#include "Components/BoxComponent.h"
+#include "StructUtils/InstancedStruct.h"
+
+#include "InteractionTargetComponent.generated.h"
+
+DECLARE_DYNAMIC_MULTICAST_DELEGATE_OneParam(FBeginInteractionCallback, const FInteractionContext&, Context);
+
+/**
+ * A target that can be interacted with by an interaction instigator.
+ * 
+ * Add this component to actors that you would like to be intractable.
+ */
+UCLASS(Blueprintable, ClassGroup = Gameplay, meta = (BlueprintSpawnableComponent), HideCategories = (Activation, AssetUserData))
+class INTERACTABLEINTERFACE_API UInteractionTargetComponent :
+	public UBoxComponent,
+	public IInteractionTarget
+{
+	GENERATED_BODY()
+public:
+
+	UInteractionTargetComponent(const FObjectInitializer& ObjectInitializer);
+
+	// A callback for when this target begins interaction.
+	UPROPERTY(BlueprintAssignable, Category="Interactions")
+	FBeginInteractionCallback OnBeginInteractionCallback;
+	
+	/**
+     * Determines what the configuration of this target is. 
+     * Gather information about this specific target so that it can be displayed
+     * to the player and provide access to what behavior should occur in response
+     * to this interaction.
+     *
+     * @param Context		The context of this interaction.
+     * @param OutResults	Output results that this target will populate
+     */
+	UFUNCTION(BlueprintCallable, BlueprintPure=false, Category="Interactions", meta=(DisplayName="Get Target Configuration", AutoCreateRefTerm="Context,OutResults"))
+	void BP_AppendTargetConfiguration(const FInteractionContext& Context, FInteractionQueryResults& OutResults) const;
+	
+	/**
+	 * Called when this target is interacted with. Implement any state changes or gameplay affects
+	 * you want this interaction to have here
+	 *
+	 * @param Context		The context of this interaction. This is customizable for your
+	 *						game by adding additional context types.
+	 */
+	UFUNCTION(BlueprintCallable, Category="Interactions", meta=(DisplayName="Begin Interaction",AutoCreateRefTerm="Context"))
+	void BP_BeginInteraction(const FInteractionContext& Context);
+	
+protected:
+
+	//~ Begin IInteractionTarget interface
+	virtual void AppendTargetConfiguration(const FInteractionContext& QueryContext, FInteractionQueryResults& OutResults) const override;
+	virtual void BeginInteraction(const FInteractionContext& Context) override;
+	//~ End of IInteractionTarget interface
+
+	/**
+	 * The configuration for this target component.
+	 *
+	 * Stores metadata about this interaction can be used to build UI or make decisions about
+	 * which target is currently desired the most by the instigator.
+	 */
+	UPROPERTY(BlueprintReadWrite, EditAnywhere, Category="Interaction", meta = (BaseStruct = "/Script/InteractableInterface.InteractionTargetConfiguration"))
+	TArray<FInstancedStruct> TargetConfigs;
+};
+
diff --git a/Engine/Plugins/Experimental/InteractionInterface/Source/InteractableInterface/Public/InteractionTask_WaitForTargets.h b/Engine/Plugins/Experimental/InteractionInterface/Source/InteractableInterface/Public/InteractionTask_WaitForTargets.h
new file mode 100644
--- /dev/null
+++ b/Engine/Plugins/Experimental/InteractionInterface/Source/InteractableInterface/Public/InteractionTask_WaitForTargets.h
@@ -0,0 +1,216 @@
+// Copyright Epic Games, Inc. All Rights Reserved.
+
+#pragma once
+
+// Engine
+#include "CoreMinimal.h"
+#include "Abilities/GameplayAbility.h"
+#include "Abilities/Tasks/AbilityTask.h"
+#include "Engine/EngineTypes.h"
+
+// Interaction Interface
+#include "InteractionTypes.h"
+
+#include "InteractionTask_WaitForTargets.generated.h"
+
+class UIndicatorDescriptor;
+class UObject;
+class UUserWidget;
+struct FFrame;
+struct FGameplayAbilityActorInfo;
+struct FGameplayEventData;
+
+/**
+ * Delegate fired for when the nearby available interactable targets have changed.
+ */
+DECLARE_DYNAMIC_MULTICAST_DELEGATE_OneParam(FAvailableInteractionTargetsChanged, const TArray<TScriptInterface<IInteractionTarget>>&, InteractableOptions);
+
+/**
+ * Gameplay Ability Task that will query all available interaction targets at a given
+ * interval around it's owning ability actor.
+ *
+ * This task will gather nearby interactive targets and make them available to you
+ * via blueprint/C++ with the OnAvailableInteractionTargetsChanged delegate
+ */
+UCLASS()
+class INTERACTABLEINTERFACE_API UAbilityTask_GrantNearbyInteractionData : public UAbilityTask
+{
+	GENERATED_BODY()
+
+public:
+	/**
+	 * Wait until an overlap occurs. This will need to be better fleshed out so we can specify game specific collision requirements
+	 */
+	UFUNCTION(BlueprintCallable, Category="Ability|Tasks", meta = (HidePin = "OwningAbility", DefaultToSelf = "OwningAbility", BlueprintInternalUseOnly = "TRUE"))
+	static UAbilityTask_GrantNearbyInteractionData* GrantAbilitiesForNearbyInteractionData(
+		UGameplayAbility* OwningAbility,
+		ECollisionChannel TraceChannel,
+		float InteractionScanRange = 500.0f,
+		float InteractionScanRate = 0.1f);
+
+	/**
+	 * Delegate fired when the available interaction targets near the owner of this ability task have changed
+	 */
+	UPROPERTY(BlueprintAssignable)
+	FAvailableInteractionTargetsChanged OnAvailableInteractionTargetsChanged;
+
+protected:
+
+	//~ UAbilityTask interface
+	virtual void Activate() override;
+	virtual void OnDestroy(bool AbilityEnded) override;
+	//~ End of UAbilityTask interface
+
+	/**
+	 * Runs a query to gather the nearby interactable objects to this ability's owner.
+	 * 
+	 */
+	virtual void QueryNearbyInteractables();
+
+	/**
+	 * How often to scan for targets. A world OverlapMultiByChannel call will happen
+	 * at this rate to check for available interactions around us.
+	 */
+	float InteractionScanRange = 500.0f;
+
+	/**
+	 * How often to scan for targets. A world OverlapMultiByChannel call will happen
+	 * at this rate to check for available interactions around us.
+	 */
+	float InteractionScanRate = 0.1f;
+
+	/**
+	 * The collision channel to check for interactable targets on
+	 */
+	ECollisionChannel InteractionTraceChannel = ECC_GameTraceChannel1;
+
+	/**
+	 * Timer handle that is populated on this task's Activate function and cleared OnDestroy
+	 * for how often to Query for targets.
+	 */
+	FTimerHandle QueryTimerHandle;
+
+	/**
+	 * Array of interaction targets which are in range from the most recent query
+	 */
+	UPROPERTY()
+	TArray<TScriptInterface<IInteractionTarget>> CurrentAvailableTargets;
+};
+
+
+
+UINTERFACE()
+class INTERACTABLEINTERFACE_API UInteractionAbilityInterface : public UInterface
+{
+	GENERATED_BODY()
+};
+
+
+class INTERACTABLEINTERFACE_API IInteractionAbilityInterface
+{
+	GENERATED_BODY()
+protected:
+
+	/**
+	 * Called when this ability's available interaction targets have been updated.
+	 * 
+	 * This is a good place to update some UI or display some message to the user
+	 * that they can now interact with the current targets.
+	 */
+	UFUNCTION(BlueprintNativeEvent, Category="Interactions")
+	void OnAvailableInteractionsUpdated();
+	virtual void OnAvailableInteractionsUpdated_Implementation() {}
+	
+	/**
+     * Triggers the interaction with one or more of the currently available targets.
+     * Override this in blueprints or native C++ to decide which of the currently available targets
+     * you would like to interact with and how.
+     */
+    UFUNCTION(BlueprintNativeEvent, Category="Interactions")
+    void OnTriggerInteraction();
+    virtual void OnTriggerInteraction_Implementation() {}
+};
+
+
+/**
+ * Gameplay ability for interacting with a target(s).
+ *
+ * This ability will trigger interactions on its current list of available targets
+ * which are populated via the UpdateInteractions functions.
+ *
+ * When UpdateInteractions is called, it provides a nice place to update
+ * some UI or other things you may want to do to display to your player
+ * that interactions are now available.
+ */
+UCLASS(Abstract)
+class INTERACTABLEINTERFACE_API UGameplayAbility_Interact
+	: public UGameplayAbility,
+	public IInteractionAbilityInterface
+{
+	GENERATED_BODY()
+
+public:
+
+	UGameplayAbility_Interact(const FObjectInitializer& ObjectInitializer = FObjectInitializer::Get());
+
+	/**
+	 * Update the available interactions that this ability can trigger.
+	 * This is normally populated via an async task running in the ability blueprint
+	 * to gather nearby targets. 
+	 */
+	UFUNCTION(BlueprintCallable, Category="Interactions")
+	void UpdateInteractions(const TArray<TScriptInterface<IInteractionTarget>>& AvailableTargets);
+
+	/**
+	 * Attempts to begin the interaction with the current targets.
+	 */
+	UFUNCTION(BlueprintCallable, Category = "Interactions")
+	void TriggerInteraction();
+
+protected:
+
+	//~ UGameplayAbility
+	virtual void ActivateAbility(
+		const FGameplayAbilitySpecHandle Handle,
+		const FGameplayAbilityActorInfo* ActorInfo,
+		const FGameplayAbilityActivationInfo ActivationInfo,
+		const FGameplayEventData* TriggerEventData)
+	override;
+	//~ End of UGameplayAbility
+
+	void HandleTargetsUpdatedFromTask(const TArray<TScriptInterface<IInteractionTarget>>& AvailableTargets);
+
+	/**
+	 * Native C++ implementation of OnAvailableInteractionsUpdated.
+	 * Override this to do any work you may need to when new interactions become available
+	 */
+	virtual void OnAvailableInteractionsUpdated_Implementation() override;
+	virtual void OnTriggerInteraction_Implementation() override;
+	
+	/**
+	 * How often to scan for targets. A world OverlapMultiByChannel call will happen
+	 * at this rate to check for available interactions around us.
+	 */
+	UPROPERTY(EditDefaultsOnly, Category = "Interactions")
+	float InteractionScanRate = 0.1f;
+
+	/**
+	 * The range to scan for available targets. A sphere of this radius will be cast around this ability's
+	 * owning actor to check for nearby interactions.
+	 */
+	UPROPERTY(EditDefaultsOnly, Category = "Interactions")
+	float InteractionScanRange = 500.0f;
+
+	/**
+	 * The collision channel to use when checking for interaction targets within the given area
+	 */
+	UPROPERTY(EditDefaultsOnly, Category = "Interactions")
+	TEnumAsByte<ECollisionChannel> InteractionTraceChannel = ECC_GameTraceChannel1;
+
+	/**
+	 * Array of available interaction targets to interact with. This is populated by "UpdateInteractions"
+	 * and normally after an ability task to gather the available targets has completed.
+	 */
+	UPROPERTY(BlueprintReadOnly, Category = "Interactions")
+	TArray<TScriptInterface<IInteractionTarget>> CurrentAvailableTargets;
+};
\ No newline at end of file
diff --git a/Engine/Plugins/Experimental/InteractionInterface/Source/InteractableInterface/Public/InteractionTypes.h b/Engine/Plugins/Experimental/InteractionInterface/Source/InteractableInterface/Public/InteractionTypes.h
new file mode 100644
--- /dev/null
+++ b/Engine/Plugins/Experimental/InteractionInterface/Source/InteractableInterface/Public/InteractionTypes.h
@@ -0,0 +1,118 @@
+// Copyright Epic Games, Inc. All Rights Reserved.
+
+#pragma once
+
+#include "StructUtils/InstancedStruct.h"
+#include "UObject/ScriptInterface.h"
+#include "GameplayTagContainer.h"
+
+#include "InteractionTypes.generated.h"
+
+class IInteractionTarget;
+class IInteractableInstigator;
+class UInteractionResponse;
+
+/**
+ * Contains data about a single interactable target. This is information that you may want to use
+ * to build some UI or decide on what the state of certain interactable objects is.
+ */
+USTRUCT(BlueprintType)
+struct FInteractionTargetConfiguration
+{
+	GENERATED_BODY()
+	
+public:
+	
+	/**
+	 * The display name that can be used for this interaction
+	 */
+	UPROPERTY(EditAnywhere, BlueprintReadWrite, Category="Interactions")
+	FText DisplayName;
+
+	// Any other metadata that we may want for a single interaction
+	// Maybe some kind of widget here, or we make a subclass of FInteractionTargetConfiguration that supports that
+	// Games can extend this type to add whatever data they need here.
+};
+
+/**
+ * Data about a specific interaction that is happening.
+ *
+ * It is encouraged that you make structs that inherit from this one to add
+ * custom game logic and conditions you may need for any given interaction.
+ * For example, you may have a target which is only able to be used if the
+ * player has a specific item equipped, or some other conditional state like that.
+ */
+USTRUCT(BlueprintType)
+struct FInteractionContextData
+{
+	GENERATED_BODY()
+
+	/**
+	 * The interaction instigator who is doing the querying!
+	 */
+	UPROPERTY(EditInstanceOnly, BlueprintReadWrite, Category="General")
+	TScriptInterface<IInteractableInstigator> Instigator;
+
+	/**
+	 * Some tags that are unique to this interaction and can be used to provide some context
+	 */
+	UPROPERTY(EditAnywhere, BlueprintReadWrite, Category="General")
+	FGameplayTagContainer InteractionTags;
+};
+
+/**
+ * The context of a given interaction query. This holds some state about
+ * what is currently querying for intractable objects and will allow you
+ * to specify some specific game state around each interaction.
+ *
+ * @see IInteractionTarget::BeginInteraction
+ */
+USTRUCT(BlueprintType)
+struct FInteractionContext
+{
+	GENERATED_BODY()
+public:
+
+	/**
+	 * The interactable target that should be used
+	 */
+	UPROPERTY(EditInstanceOnly, BlueprintReadWrite, Category=General)
+	TScriptInterface<IInteractionTarget> Target;
+	
+	/**
+	 * Data about this specific interaction query
+	 */
+	UPROPERTY(EditAnywhere, BlueprintReadWrite, Category = General, meta = (BaseStruct = "/Script/InteractableInterface.InteractionContextData"))
+	FInstancedStruct ContextData;
+};
+
+/**
+ * A struct that will be populated with all the available interaction options for a given target.
+ *
+ * These query results are populated by Interaction Targets. Each target may have
+ * multiple configurations on it (i.e. a bench with two seats, or a car with 4 doors to open)
+ *
+ * @see IInteractionTarget::GetTargetConfiguration
+ */
+USTRUCT(BlueprintType)
+struct FInteractionQueryResults
+{
+	GENERATED_BODY()
+public:	
+
+	/**
+	* Resets the values of this query results to be empty.
+	*/
+	void Reset()
+	{
+		AvailableInteractions.Reset();
+	}
+
+	/**
+	 * Array of available interactions that can be started.
+	 * Add to this array for any interaction which you would like to be presented
+	 * as an available option in response to this query.
+	 */
+	UPROPERTY(EditAnywhere, BlueprintReadWrite, Category = General, meta = (BaseStruct = "/Script/InteractableInterface.InteractionTargetConfiguration"))
+	TArray<FInstancedStruct> AvailableInteractions;
+};
